Added GAS Toolsets

#Synced-CL 52186709

diff --git a/Engine/Plugins/Experimental/Toolsets/GASToolsets/GASToolsets.uplugin b/Engine/Plugins/Experimental/Toolsets/GASToolsets/GASToolsets.uplugin
new file mode 100644
--- /dev/null
+++ b/Engine/Plugins/Experimental/Toolsets/GASToolsets/GASToolsets.uplugin
@@ -0,0 +1,38 @@
+{
+	"FileVersion": 3,
+	"Version": 1,
+	"VersionName": "1.0",
+	"FriendlyName": "GASToolsets",
+	"Description": "Toolsets for the Gameplay Ability System",
+	"Category": "Other",
+	"EditorOnly": true,
+	"CreatedBy": "Epic Games",
+	"CreatedByURL": "",
+	"DocsURL": "",
+	"MarketplaceURL": "",
+	"SupportURL": "",
+	"CanContainContent": false,
+	"IsBetaVersion": false,
+	"IsExperimentalVersion": true,
+	"NoRedist": true,
+	"Installed": false,
+	"EnabledByDefault": true,
+	"Modules": [
+		{
+			"Name": "GASToolsets",
+			"Type": "Editor",
+			"LoadingPhase": "PostEngineInit",
+			"TargetAllowList": ["Editor"]
+		}
+	],
+	"Plugins": [
+		{
+			"Name": "GameplayAbilities",
+			"Enabled": true
+		},
+		{
+			"Name": "ToolsetRegistry",
+			"Enabled": true
+		}
+	]
+}
diff --git a/Engine/Plugins/Experimental/Toolsets/GASToolsets/Source/GASToolsets/GASToolsets.Build.cs b/Engine/Plugins/Experimental/Toolsets/GASToolsets/Source/GASToolsets/GASToolsets.Build.cs
new file mode 100644
--- /dev/null
+++ b/Engine/Plugins/Experimental/Toolsets/GASToolsets/Source/GASToolsets/GASToolsets.Build.cs
@@ -0,0 +1,27 @@
+// Copyright Epic Games, Inc. All Rights Reserved.
+
+using UnrealBuildTool;
+
+public class GASToolsets : ModuleRules
+{
+	public GASToolsets(ReadOnlyTargetRules Target) : base(Target)
+	{
+		PCHUsage = PCHUsageMode.UseExplicitOrSharedPCHs;
+
+		PublicDependencyModuleNames.Add("Core");
+
+		PrivateDependencyModuleNames.AddRange(new string[]
+		{
+			"AssetRegistry",
+			"CoreUObject",
+			"Engine",
+			"GameplayAbilities",
+			"GameplayTags",
+			"GameplayTagsEditor",
+			"Kismet",
+			"KismetCompiler",
+			"ToolsetRegistry",
+			"UnrealEd",
+		});
+	}
+}
diff --git a/Engine/Plugins/Experimental/Toolsets/GASToolsets/Source/GASToolsets/Private/AbilitySystemInspectorToolset.cpp b/Engine/Plugins/Experimental/Toolsets/GASToolsets/Source/GASToolsets/Private/AbilitySystemInspectorToolset.cpp
new file mode 100644
--- /dev/null
+++ b/Engine/Plugins/Experimental/Toolsets/GASToolsets/Source/GASToolsets/Private/AbilitySystemInspectorToolset.cpp
@@ -0,0 +1,191 @@
+// Copyright Epic Games, Inc. All Rights Reserved.
+
+#include "AbilitySystemInspectorToolset.h"
+
+#include "AbilitySystemBlueprintLibrary.h"
+#include "AbilitySystemComponent.h"
+#include "AbilitySystemGlobals.h"
+#include "AttributeSet.h"
+#include "Engine/World.h"
+#include "GameplayEffect.h"
+#include "GameplayEffectTypes.h"
+#include "Kismet/KismetSystemLibrary.h"
+
+#include UE_INLINE_GENERATED_CPP_BY_NAME(AbilitySystemInspectorToolset)
+
+namespace
+{
+	/**
+	 * Returns the AbilitySystemComponent for the given actor.
+	 * Raises a script error and returns nullptr if Actor is null or has no ASC.
+	 */
+	UAbilitySystemComponent* GetASCFromActor(AActor* Actor)
+	{
+		if (!Actor)
+		{
+			UKismetSystemLibrary::RaiseScriptError(TEXT("Actor is null."));
+			return nullptr;
+		}
+
+		UAbilitySystemComponent* ASC =
+			UAbilitySystemGlobals::GetAbilitySystemComponentFromActor(Actor);
+
+		if (!ASC)
+		{
+			UKismetSystemLibrary::RaiseScriptError(FString::Printf(
+				TEXT("Actor '%s' does not have an AbilitySystemComponent."),
+				*Actor->GetActorLabel()));
+			return nullptr;
+		}
+
+		return ASC;
+	}
+}
+
+TArray<FRuntimeAttributeValue> UAbilitySystemInspectorToolset::GetAttributeValues(AActor* Actor)
+{
+	UAbilitySystemComponent* ASC = GetASCFromActor(Actor);
+	if (!ASC)
+	{
+		return {};
+	}
+
+	AActor* OwnerActor = ASC->GetOwner();
+
+	TArray<FRuntimeAttributeValue> Results;
+
+	for (const UAttributeSet* Set : ASC->GetSpawnedAttributes())
+	{
+		if (!Set)
+		{
+			continue;
+		}
+
+		UClass* SetClass = Set->GetClass();
+		const FString SetClassName = SetClass->GetName();
+
+		TArray<FGameplayAttribute> Attributes;
+		UAttributeSet::GetAttributesFromSetClass(SetClass, Attributes);
+
+		for (const FGameplayAttribute& Attr : Attributes)
+		{
+			bool bFoundBase = false;
+			bool bFoundCurrent = false;
+			const float BaseVal =
+				UAbilitySystemBlueprintLibrary::GetFloatAttributeBase(OwnerActor, Attr, bFoundBase);
+			const float CurrentVal =
+				UAbilitySystemBlueprintLibrary::GetFloatAttribute(OwnerActor, Attr, bFoundCurrent);
+
+			FRuntimeAttributeValue Info;
+			Info.AttributeName = Attr.GetName();
+			Info.SetClassName = SetClassName;
+			Info.FullName = FString::Printf(TEXT("%s.%s"), *SetClassName, *Info.AttributeName);
+			Info.BaseValue = bFoundBase ? BaseVal : 0.f;
+			Info.CurrentValue = bFoundCurrent ? CurrentVal : 0.f;
+			Results.Add(Info);
+		}
+	}
+
+	Results.Sort([](const FRuntimeAttributeValue& A, const FRuntimeAttributeValue& B)
+	{
+		return A.FullName < B.FullName;
+	});
+
+	return Results;
+}
+
+TArray<FActiveEffectInfo> UAbilitySystemInspectorToolset::GetActiveEffects(AActor* Actor)
+{
+	UAbilitySystemComponent* ASC = GetASCFromActor(Actor);
+	if (!ASC)
+	{
+		return {};
+	}
+
+	TArray<FActiveEffectInfo> Results;
+
+	const float WorldTime = ASC->GetWorld() ? ASC->GetWorld()->GetTimeSeconds() : 0.f;
+
+	const FActiveGameplayEffectsContainer& Container = ASC->GetActiveGameplayEffects();
+	for (const FActiveGameplayEffect& ActiveEffect : &Container)
+	{
+		if (ActiveEffect.IsPendingRemove)
+		{
+			continue;
+		}
+
+		const UGameplayEffect* GEDef = ActiveEffect.Spec.Def.Get();
+
+		FActiveEffectInfo Info;
+		Info.EffectName = GEDef ? GEDef->GetName() : TEXT("Unknown");
+		Info.StackCount =
+			UAbilitySystemBlueprintLibrary::GetActiveGameplayEffectStackCount(ActiveEffect.Handle);
+		Info.TotalDuration = ActiveEffect.GetDuration();
+		Info.RemainingDuration = ActiveEffect.GetTimeRemaining(WorldTime);
+
+		// Collect the tags the GE definition grants to the owner.
+		if (GEDef)
+		{
+			const FGameplayTagContainer& GrantedTags =
+				UAbilitySystemBlueprintLibrary::GetGameplayEffectGrantedTags(GEDef->GetClass());
+			for (const FGameplayTag& Tag : GrantedTags)
+			{
+				Info.GrantedTags.Add(Tag.ToString());
+			}
+		}
+
+		Results.Add(Info);
+	}
+
+	Results.Sort([](const FActiveEffectInfo& A, const FActiveEffectInfo& B)
+	{
+		return A.EffectName < B.EffectName;
+	});
+
+	return Results;
+}
+
+TArray<FGrantedAbilityInfo> UAbilitySystemInspectorToolset::GetGrantedAbilities(AActor* Actor)
+{
+	UAbilitySystemComponent* ASC = GetASCFromActor(Actor);
+	if (!ASC)
+	{
+		return {};
+	}
+
+	TArray<FGrantedAbilityInfo> Results;
+
+	for (const FGameplayAbilitySpec& Spec : ASC->GetActivatableAbilities())
+	{
+		FGrantedAbilityInfo Info;
+		Info.AbilityName = Spec.Ability ? Spec.Ability->GetClass()->GetName() : TEXT("Unknown");
+		Info.Level = Spec.Level;
+		Info.bIsActive = Spec.IsActive();
+		Results.Add(Info);
+	}
+
+	Results.Sort([](const FGrantedAbilityInfo& A, const FGrantedAbilityInfo& B)
+	{
+		return A.AbilityName < B.AbilityName;
+	});
+
+	return Results;
+}
+
+TArray<FString> UAbilitySystemInspectorToolset::GetActiveTags(AActor* Actor)
+{
+	UAbilitySystemComponent* ASC = GetASCFromActor(Actor);
+	if (!ASC)
+	{
+		return {};
+	}
+
+	TArray<FString> TagStrings;
+	for (const FGameplayTag& Tag : ASC->GetOwnedGameplayTags())
+	{
+		TagStrings.Add(Tag.ToString());
+	}
+
+	TagStrings.Sort();
+	return TagStrings;
+}
diff --git a/Engine/Plugins/Experimental/Toolsets/GASToolsets/Source/GASToolsets/Private/AbilitySystemInspectorToolset.h b/Engine/Plugins/Experimental/Toolsets/GASToolsets/Source/GASToolsets/Private/AbilitySystemInspectorToolset.h
new file mode 100644
--- /dev/null
+++ b/Engine/Plugins/Experimental/Toolsets/GASToolsets/Source/GASToolsets/Private/AbilitySystemInspectorToolset.h
@@ -0,0 +1,134 @@
+// Copyright Epic Games, Inc. All Rights Reserved.
+
+#pragma once
+
+#include "CoreMinimal.h"
+#include "ToolsetRegistry/UToolsetRegistry.h"
+
+#include "AbilitySystemInspectorToolset.generated.h"
+
+/// The runtime base and current value of a single gameplay attribute.
+USTRUCT(BlueprintType)
+struct FRuntimeAttributeValue
+{
+	GENERATED_BODY()
+
+	/// The short attribute name (e.g. "Health").
+	UPROPERTY(BlueprintReadWrite, Category = "AbilitySystem")
+	FString AttributeName;
+
+	/// The fully-qualified name including the owning set (e.g. "UMyHealthSet.Health").
+	UPROPERTY(BlueprintReadWrite, Category = "AbilitySystem")
+	FString FullName;
+
+	/// The name of the AttributeSet class that owns this attribute (e.g. "UMyHealthSet").
+	UPROPERTY(BlueprintReadWrite, Category = "AbilitySystem")
+	FString SetClassName;
+
+	/// The permanent base value, unaffected by temporary modifiers.
+	UPROPERTY(BlueprintReadWrite, Category = "AbilitySystem")
+	float BaseValue = 0.f;
+
+	/// The current value after all active modifiers are applied.
+	UPROPERTY(BlueprintReadWrite, Category = "AbilitySystem")
+	float CurrentValue = 0.f;
+};
+
+/// Summarizes a gameplay effect that is currently active on an AbilitySystemComponent.
+USTRUCT(BlueprintType)
+struct FActiveEffectInfo
+{
+	GENERATED_BODY()
+
+	/// The name of the GameplayEffect class (e.g. "GE_Burning").
+	UPROPERTY(BlueprintReadWrite, Category = "AbilitySystem")
+	FString EffectName;
+
+	/// Current stack count. 1 for non-stacking effects.
+	UPROPERTY(BlueprintReadWrite, Category = "AbilitySystem")
+	int32 StackCount = 1;
+
+	/// Total duration in seconds. -1 means infinite.
+	UPROPERTY(BlueprintReadWrite, Category = "AbilitySystem")
+	float TotalDuration = 0.f;
+
+	/// Remaining duration in seconds. -1 means infinite.
+	UPROPERTY(BlueprintReadWrite, Category = "AbilitySystem")
+	float RemainingDuration = 0.f;
+
+	/// Tags granted to the owner by this effect.
+	UPROPERTY(BlueprintReadWrite, Category = "AbilitySystem")
+	TArray<FString> GrantedTags;
+};
+
+/// Summarizes an ability that has been granted to an AbilitySystemComponent.
+USTRUCT(BlueprintType)
+struct FGrantedAbilityInfo
+{
+	GENERATED_BODY()
+
+	/// The name of the ability class (e.g. "GA_Sprint").
+	UPROPERTY(BlueprintReadWrite, Category = "AbilitySystem")
+	FString AbilityName;
+
+	/// The level at which this ability was granted.
+	UPROPERTY(BlueprintReadWrite, Category = "AbilitySystem")
+	int32 Level = 1;
+
+	/// True if the ability is currently executing (at least one active instance).
+	UPROPERTY(BlueprintReadWrite, Category = "AbilitySystem")
+	bool bIsActive = false;
+};
+
+/**
+ * Provides tools for inspecting the runtime state of an AbilitySystemComponent.
+ * Each function takes a direct actor pointer and raises a script error if the
+ * actor is null or has no AbilitySystemComponent.
+ */
+UCLASS(BlueprintType, Hidden)
+class UAbilitySystemInspectorToolset : public UToolsetDefinition
+{
+	GENERATED_BODY()
+
+public:
+	/**
+	 * Returns the current base and modified values of all gameplay attributes
+	 * on the actor's AbilitySystemComponent.
+	 * @param Actor  The target actor.
+	 * @return A list of attribute values sorted by full name.
+	 *   Raises a script error if Actor is null or has no ASC.
+	 */
+	UFUNCTION(BlueprintCallable, Category = "AbilitySystem")
+	static TArray<FRuntimeAttributeValue> GetAttributeValues(AActor* Actor);
+
+	/**
+	 * Returns all gameplay effects currently active on the actor's
+	 * AbilitySystemComponent.
+	 * @param Actor  The target actor.
+	 * @return A list of active effect summaries sorted by effect name.
+	 *   Raises a script error if Actor is null or has no ASC.
+	 */
+	UFUNCTION(BlueprintCallable, Category = "AbilitySystem")
+	static TArray<FActiveEffectInfo> GetActiveEffects(AActor* Actor);
+
+	/**
+	 * Returns all abilities granted to the actor's AbilitySystemComponent.
+	 * @param Actor  The target actor.
+	 * @return A list of granted ability summaries sorted by ability name.
+	 *   Raises a script error if Actor is null or has no ASC.
+	 */
+	UFUNCTION(BlueprintCallable, Category = "AbilitySystem")
+	static TArray<FGrantedAbilityInfo> GetGrantedAbilities(AActor* Actor);
+
+	/**
+	 * Returns the gameplay tags currently owned by the actor's
+	 * AbilitySystemComponent (includes loose tags, effect-granted tags, etc.).
+	 * @param Actor  The target actor.
+	 * @return A sorted list of gameplay tag name strings.
+	 *   Raises a script error if Actor is null or has no ASC.
+	 */
+	UFUNCTION(BlueprintCallable, Category = "AbilitySystem")
+	static TArray<FString> GetActiveTags(AActor* Actor);
+
+	friend class FAbilitySystemInspectorToolsetSpec;
+};
diff --git a/Engine/Plugins/Experimental/Toolsets/GASToolsets/Source/GASToolsets/Private/AttributeSetToolset.cpp b/Engine/Plugins/Experimental/Toolsets/GASToolsets/Source/GASToolsets/Private/AttributeSetToolset.cpp
new file mode 100644
--- /dev/null
+++ b/Engine/Plugins/Experimental/Toolsets/GASToolsets/Source/GASToolsets/Private/AttributeSetToolset.cpp
@@ -0,0 +1,110 @@
+// Copyright Epic Games, Inc. All Rights Reserved.
+
+#include "AttributeSetToolset.h"
+
+#include "AssetRegistry/AssetRegistryModule.h"
+#include "AttributeSet.h"
+#include "Engine/Blueprint.h"
+#include "Engine/BlueprintGeneratedClass.h"
+#include "Kismet/KismetSystemLibrary.h"
+#include "Modules/ModuleManager.h"
+#include "UObject/UObjectIterator.h"
+
+#include UE_INLINE_GENERATED_CPP_BY_NAME(AttributeSetToolset)
+
+namespace
+{
+	FAttributeSetClassInfo BuildClassInfo(UClass* Class)
+	{
+		FAttributeSetClassInfo Info;
+		Info.ClassName = Class->GetName();
+
+		// Populate AssetPath for Blueprint-generated classes.
+		if (UBlueprintGeneratedClass* BPGC = Cast<UBlueprintGeneratedClass>(Class))
+		{
+			if (UBlueprint* BP = Cast<UBlueprint>(BPGC->ClassGeneratedBy))
+			{
+				Info.AssetPath = BP->GetPathName();
+			}
+		}
+
+		TArray<FGameplayAttribute> Attributes;
+		UAttributeSet::GetAttributesFromSetClass(Class, Attributes);
+
+		for (const FGameplayAttribute& Attr : Attributes)
+		{
+			FGameplayAttributeInfo AttrInfo;
+			AttrInfo.AttributeName = Attr.GetName();
+			AttrInfo.SetClassName = Info.ClassName;
+			AttrInfo.FullName = FString::Printf(TEXT("%s.%s"), *Info.ClassName, *AttrInfo.AttributeName);
+			Info.Attributes.Add(AttrInfo);
+		}
+
+		return Info;
+	}
+}
+
+TArray<FAttributeSetClassInfo> UAttributeSetToolset::FindAttributeSetClasses()
+{
+	TArray<FAttributeSetClassInfo> Results;
+
+	// Collect all loaded UClass objects that are strict subclasses of UAttributeSet.
+	for (TObjectIterator<UClass> ClassIt; ClassIt; ++ClassIt)
+	{
+		UClass* Class = *ClassIt;
+		if (!Class->IsChildOf(UAttributeSet::StaticClass()) ||
+			Class == UAttributeSet::StaticClass() ||
+			Class->HasAnyClassFlags(CLASS_Abstract))
+		{
+			continue;
+		}
+
+		Results.Add(BuildClassInfo(Class));
+	}
+
+	Results.Sort([](const FAttributeSetClassInfo& A, const FAttributeSetClassInfo& B)
+	{
+		return A.ClassName < B.ClassName;
+	});
+
+	return Results;
+}
+
+TArray<FGameplayAttributeInfo> UAttributeSetToolset::ListAttributes(const FString& ClassName)
+{
+	// Search loaded classes by name.
+	UClass* FoundClass = nullptr;
+	for (TObjectIterator<UClass> ClassIt; ClassIt; ++ClassIt)
+	{
+		UClass* Class = *ClassIt;
+		if (Class->GetName() == ClassName && Class->IsChildOf(UAttributeSet::StaticClass()))
+		{
+			FoundClass = Class;
+			break;
+		}
+	}
+
+	if (!FoundClass)
+	{
+		UKismetSystemLibrary::RaiseScriptError(FString::Printf(
+			TEXT("AttributeSet class '%s' not found. Use FindAttributeSetClasses to list available classes."),
+			*ClassName));
+		return {};
+	}
+
+	TArray<FGameplayAttribute> Attributes;
+	UAttributeSet::GetAttributesFromSetClass(FoundClass, Attributes);
+
+	TArray<FGameplayAttributeInfo> Results;
+	Results.Reserve(Attributes.Num());
+	for (const FGameplayAttribute& Attr : Attributes)
+	{
+		FGameplayAttributeInfo Info;
+		Info.AttributeName = Attr.GetName();
+		Info.SetClassName = ClassName;
+		Info.FullName = FString::Printf(TEXT("%s.%s"), *ClassName, *Info.AttributeName);
+		Results.Add(Info);
+	}
+
+	return Results;
+}
diff --git a/Engine/Plugins/Experimental/Toolsets/GASToolsets/Source/GASToolsets/Private/AttributeSetToolset.h b/Engine/Plugins/Experimental/Toolsets/GASToolsets/Source/GASToolsets/Private/AttributeSetToolset.h
new file mode 100644
--- /dev/null
+++ b/Engine/Plugins/Experimental/Toolsets/GASToolsets/Source/GASToolsets/Private/AttributeSetToolset.h
@@ -0,0 +1,73 @@
+// Copyright Epic Games, Inc. All Rights Reserved.
+
+#pragma once
+
+#include "CoreMinimal.h"
+#include "ToolsetRegistry/UToolsetRegistry.h"
+
+#include "AttributeSetToolset.generated.h"
+
+/// Describes a gameplay attribute belonging to an AttributeSet class.
+USTRUCT(BlueprintType)
+struct FGameplayAttributeInfo
+{
+	GENERATED_BODY()
+
+	/// The short attribute name (e.g. "Health").
+	UPROPERTY(BlueprintReadWrite, Category = "AttributeSet")
+	FString AttributeName;
+
+	/// The fully-qualified attribute name including the set class (e.g. "UMyHealthSet.Health").
+	UPROPERTY(BlueprintReadWrite, Category = "AttributeSet")
+	FString FullName;
+
+	/// The name of the AttributeSet class that owns this attribute (e.g. "UMyHealthSet").
+	UPROPERTY(BlueprintReadWrite, Category = "AttributeSet")
+	FString SetClassName;
+};
+
+/// Describes an AttributeSet subclass found in the project.
+USTRUCT(BlueprintType)
+struct FAttributeSetClassInfo
+{
+	GENERATED_BODY()
+
+	/// The UClass name of the attribute set (e.g. "UMyHealthSet").
+	UPROPERTY(BlueprintReadWrite, Category = "AttributeSet")
+	FString ClassName;
+
+	/// The object path to the Blueprint asset, or empty for native C++ classes.
+	UPROPERTY(BlueprintReadWrite, Category = "AttributeSet")
+	FString AssetPath;
+
+	/// The gameplay attributes defined on this set.
+	UPROPERTY(BlueprintReadWrite, Category = "AttributeSet")
+	TArray<FGameplayAttributeInfo> Attributes;
+};
+
+/// Provides tools for discovering AttributeSet classes and their gameplay attributes.
+UCLASS(BlueprintType, Hidden)
+class UAttributeSetToolset : public UToolsetDefinition
+{
+	GENERATED_BODY()
+
+public:
+	/**
+	 * Returns all AttributeSet subclasses found in the project, including their attributes.
+	 * Covers both native C++ subclasses and Blueprint subclasses discovered via the asset registry.
+	 * @return A list of attribute set descriptors sorted by class name.
+	 */
+	UFUNCTION(BlueprintCallable, Category = "AttributeSets")
+	static TArray<FAttributeSetClassInfo> FindAttributeSetClasses();
+
+	/**
+	 * Returns the gameplay attributes defined on a specific AttributeSet class.
+	 * @param ClassName The UClass name to look up, e.g. "UMyHealthSet".
+	 *   Raises a script error if the class is not found or is not an AttributeSet subclass.
+	 * @return A list of attribute descriptors defined on the class.
+	 */
+	UFUNCTION(BlueprintCallable, Category = "AttributeSets")
+	static TArray<FGameplayAttributeInfo> ListAttributes(const FString& ClassName);
+
+	friend class FAttributeSetToolsetSpec;
+};
diff --git a/Engine/Plugins/Experimental/Toolsets/GASToolsets/Source/GASToolsets/Private/GASToolsets.cpp b/Engine/Plugins/Experimental/Toolsets/GASToolsets/Source/GASToolsets/Private/GASToolsets.cpp
new file mode 100644
--- /dev/null
+++ b/Engine/Plugins/Experimental/Toolsets/GASToolsets/Source/GASToolsets/Private/GASToolsets.cpp
@@ -0,0 +1,28 @@
+// Copyright Epic Games, Inc. All Rights Reserved.
+
+#include "GASToolsets.h"
+
+#include "AbilitySystemInspectorToolset.h"
+#include "AttributeSetToolset.h"
+#include "Editor.h"
+#include "GameplayCueToolset.h"
+#include "Misc/CoreDelegates.h"
+#include "Modules/ModuleManager.h"
+#include "ToolsetRegistry/UToolsetRegistry.h"
+
+/*virtual*/ void FGASToolsetsModule::StartupModule()
+{
+	UToolsetRegistry::RegisterToolsetClass(UGameplayCueToolset::StaticClass());
+	UToolsetRegistry::RegisterToolsetClass(UAttributeSetToolset::StaticClass());
+	UToolsetRegistry::RegisterToolsetClass(UAbilitySystemInspectorToolset::StaticClass());
+}
+
+/*virtual*/ void FGASToolsetsModule::ShutdownModule()
+{
+	UToolsetRegistry::UnregisterToolsetClass(UAbilitySystemInspectorToolset::StaticClass());
+	UToolsetRegistry::UnregisterToolsetClass(UAttributeSetToolset::StaticClass());
+	UToolsetRegistry::UnregisterToolsetClass(UGameplayCueToolset::StaticClass());
+}
+
+IMPLEMENT_MODULE(FGASToolsetsModule, GASToolsets);
+DEFINE_LOG_CATEGORY(LogGASToolsets);
diff --git a/Engine/Plugins/Experimental/Toolsets/GASToolsets/Source/GASToolsets/Private/GameplayCueToolset.cpp b/Engine/Plugins/Experimental/Toolsets/GASToolsets/Source/GASToolsets/Private/GameplayCueToolset.cpp
new file mode 100644
--- /dev/null
+++ b/Engine/Plugins/Experimental/Toolsets/GASToolsets/Source/GASToolsets/Private/GameplayCueToolset.cpp
@@ -0,0 +1,367 @@
+// Copyright Epic Games, Inc. All Rights Reserved.
+
+#include "GameplayCueToolset.h"
+
+#include "AbilitySystemGlobals.h"
+#include "AssetRegistry/AssetRegistryModule.h"
+#include "Editor.h"
+#include "Engine/Blueprint.h"
+#include "Engine/BlueprintGeneratedClass.h"
+#include "GameplayCueManager.h"
+#include "GameplayCueNotify_Actor.h"
+#include "GameplayCueNotify_Static.h"
+#include "GameplayTagsEditorModule.h"
+#include "GameplayTagsManager.h"
+#include "Kismet/KismetSystemLibrary.h"
+#include "Kismet2/KismetEditorUtilities.h"
+#include "Modules/ModuleManager.h"
+#include "Selection.h"
+
+#include UE_INLINE_GENERATED_CPP_BY_NAME(GameplayCueToolset)
+
+namespace
+{
+	static const FString GameplayCueRootTag = TEXT("GameplayCue");
+	static const FName GameplayCueNameProperty = FName("GameplayCueName");
+
+	TArray<FString> CollectDescendantTagNames(const TSharedPtr<FGameplayTagNode>& ParentNode)
+	{
+		TArray<FString> TagNames;
+		TArray<TSharedPtr<FGameplayTagNode>> Queue = ParentNode->GetChildTagNodes();
+		while (Queue.Num() > 0)
+		{
+			TSharedPtr<FGameplayTagNode> Node = Queue.Pop(EAllowShrinking::No);
+			if (Node.IsValid())
+			{
+				TagNames.Add(Node->GetCompleteTagName().ToString());
+				Queue.Append(Node->GetChildTagNodes());
+			}
+		}
+		return TagNames;
+	}
+
+	void CollectNotifyAssets(
+		UClass* BaseClass, const FString& NotifyType,
+		TArray<FGameplayCueNotifyInfo>& OutInfos)
+	{
+		IAssetRegistry& Registry =
+			FModuleManager::LoadModuleChecked<FAssetRegistryModule>("AssetRegistry").Get();
+
+		FARFilter Filter;
+		Filter.ClassPaths.Add(BaseClass->GetClassPathName());
+		Filter.bRecursiveClasses = true;
+		Filter.bRecursivePaths = true;
+
+		TArray<FAssetData> AssetDataList;
+		Registry.GetAssets(Filter, AssetDataList);
+
+		for (const FAssetData& Data : AssetDataList)
+		{
+			FString CueTagValue;
+			if (!Data.GetTagValue(GameplayCueNameProperty, CueTagValue) || CueTagValue.IsEmpty())
+			{
+				continue;
+			}
+
+			FGameplayCueNotifyInfo& Info = OutInfos.AddDefaulted_GetRef();
+			Info.CueTag = CueTagValue;
+			Info.AssetPath = Data.GetSoftObjectPath().ToString();
+			Info.AssetName = Data.AssetName.ToString();
+			Info.NotifyType = NotifyType;
+		}
+	}
+}
+
+TArray<FString> UGameplayCueToolset::ListCues(const FString& ParentTag)
+{
+	UGameplayTagsManager& TagManager = UGameplayTagsManager::Get();
+
+	const FString& RootName = ParentTag.IsEmpty() ? GameplayCueRootTag : ParentTag;
+	TSharedPtr<FGameplayTagNode> RootNode = TagManager.FindTagNode(FName(*RootName));
+
+	if (!RootNode.IsValid())
+	{
+		if (!ParentTag.IsEmpty())
+		{
+			UKismetSystemLibrary::RaiseScriptError(
+				FString::Printf(TEXT("Gameplay tag '%s' does not exist."), *ParentTag));
+		}
+		return {};
+	}
+
+	TArray<FString> TagNames = CollectDescendantTagNames(RootNode);
+	TagNames.Sort();
+	return TagNames;
+}
+
+FGameplayCueInfo UGameplayCueToolset::GetCueInfo(const FString& CueTag)
+{
+	UGameplayTagsManager& TagManager = UGameplayTagsManager::Get();
+	if (!TagManager.FindTagNode(FName(*CueTag)).IsValid())
+	{
+		UKismetSystemLibrary::RaiseScriptError(
+			FString::Printf(TEXT("Gameplay cue tag '%s' does not exist."), *CueTag));
+		return FGameplayCueInfo();
+	}
+
+	FGameplayCueInfo Info;
+	Info.Tag = CueTag;
+	Info.NotifyType = TEXT("None");
+
+	IAssetRegistry& Registry =
+		FModuleManager::LoadModuleChecked<FAssetRegistryModule>("AssetRegistry").Get();
+
+	const auto TryFindNotify = [&Registry, &CueTag, &Info](UClass* BaseClass, const FString& Type) -> bool
+	{
+		FARFilter Filter;
+		Filter.ClassPaths.Add(BaseClass->GetClassPathName());
+		Filter.bRecursiveClasses = true;
+		Filter.TagsAndValues.Add(GameplayCueNameProperty, CueTag);
+
+		TArray<FAssetData> Results;
+		Registry.GetAssets(Filter, Results);
+
+		if (Results.Num() > 0)
+		{
+			Info.NotifyAssetPath = Results[0].GetSoftObjectPath().ToString();
+			Info.NotifyType = Type;
+			return true;
+		}
+		return false;
+	};
+
+	if (!TryFindNotify(UGameplayCueNotify_Static::StaticClass(), TEXT("Static")))
+	{
+		TryFindNotify(AGameplayCueNotify_Actor::StaticClass(), TEXT("Actor"));
+	}
+
+	return Info;
+}
+
+bool UGameplayCueToolset::ExecuteCueOnSelectedActor(
+	const FString& CueTag, float NormalizedMagnitude, FVector Location, FVector Normal)
+{
+	if (!GEditor)
+	{
+		UKismetSystemLibrary::RaiseScriptError(TEXT("Editor not found."));
+		return false;
+	}
+
+	AActor* TargetActor = nullptr;
+	for (FSelectionIterator Iter(*GEditor->GetSelectedActors()); Iter; ++Iter)
+	{
+		if (AActor* Actor = Cast<AActor>(*Iter))
+		{
+			TargetActor = Actor;
+			break;
+		}
+	}
+
+	if (!TargetActor)
+	{
+		UKismetSystemLibrary::RaiseScriptError(
+			TEXT("No actor selected. Select an actor in the editor before executing the cue."));
+		return false;
+	}
+
+	const FGameplayTag Tag =
+		UGameplayTagsManager::Get().RequestGameplayTag(FName(*CueTag), false);
+	if (!Tag.IsValid())
+	{
+		UKismetSystemLibrary::RaiseScriptError(
+			FString::Printf(TEXT("Gameplay cue tag '%s' does not exist."), *CueTag));
+		return false;
+	}
+
+	FGameplayCueParameters Params;
+	Params.NormalizedMagnitude = NormalizedMagnitude;
+	Params.Location = Location;
+	Params.Normal = Normal;
+
+	UGameplayCueManager::ExecuteGameplayCue_NonReplicated(TargetActor, Tag, Params);
+	return true;
+}
+
+TArray<FGameplayCueNotifyInfo> UGameplayCueToolset::FindCueNotifyAssets(const FString& ParentTag)
+{
+	TArray<FGameplayCueNotifyInfo> Infos;
+	CollectNotifyAssets(UGameplayCueNotify_Static::StaticClass(), TEXT("Static"), Infos);
+	CollectNotifyAssets(AGameplayCueNotify_Actor::StaticClass(), TEXT("Actor"), Infos);
+
+	if (!ParentTag.IsEmpty())
+	{
+		UGameplayTagsManager& TagManager = UGameplayTagsManager::Get();
+		const FGameplayTag FilterTag = TagManager.RequestGameplayTag(FName(*ParentTag), false);
+		if (!FilterTag.IsValid())
+		{
+			UKismetSystemLibrary::RaiseScriptError(
+				FString::Printf(TEXT("Parent tag '%s' does not exist."), *ParentTag));
+			return {};
+		}
+
+		Infos.RemoveAll([&TagManager, FilterTag](const FGameplayCueNotifyInfo& Item)
+		{
+			const FGameplayTag ItemTag =
+				TagManager.RequestGameplayTag(FName(*Item.CueTag), false);
+			return !ItemTag.IsValid() || !ItemTag.MatchesTag(FilterTag);
+		});
+	}
+
+	Infos.Sort([](const FGameplayCueNotifyInfo& A, const FGameplayCueNotifyInfo& B)
+	{
+		return A.CueTag < B.CueTag;
+	});
+
+	return Infos;
+}
+
+FString UGameplayCueToolset::CreateCueNotifyAsset(
+	const FString& CueTag, const FString& PackagePath,
+	const FString& AssetName, bool bIsActor)
+{
+	if (CueTag.IsEmpty())
+	{
+		UKismetSystemLibrary::RaiseScriptError(TEXT("CueTag cannot be empty."));
+		return FString();
+	}
+	if (PackagePath.IsEmpty())
+	{
+		UKismetSystemLibrary::RaiseScriptError(TEXT("PackagePath cannot be empty."));
+		return FString();
+	}
+	if (AssetName.IsEmpty())
+	{
+		UKismetSystemLibrary::RaiseScriptError(TEXT("AssetName cannot be empty."));
+		return FString();
+	}
+
+	const FGameplayTag Tag =
+		UGameplayTagsManager::Get().RequestGameplayTag(FName(*CueTag), false);
+	if (!Tag.IsValid())
+	{
+		UKismetSystemLibrary::RaiseScriptError(FString::Printf(
+			TEXT("Gameplay cue tag '%s' does not exist. Add it first with AddCueTag."), *CueTag));
+		return FString();
+	}
+
+	const FString FullPackagePath = PackagePath / AssetName;
+	UPackage* Package = CreatePackage(*FullPackagePath);
+	if (!Package)
+	{
+		UKismetSystemLibrary::RaiseScriptError(
+			FString::Printf(TEXT("Failed to create package at '%s'."), *FullPackagePath));
+		return FString();
+	}
+
+	UClass* NotifyClass = bIsActor
+		? AGameplayCueNotify_Actor::StaticClass()
+		: UGameplayCueNotify_Static::StaticClass();
+
+	UBlueprint* NewBP = FKismetEditorUtilities::CreateBlueprint(
+		NotifyClass, Package, FName(*AssetName),
+		BPTYPE_Normal, UBlueprint::StaticClass(), UBlueprintGeneratedClass::StaticClass(),
+		FName("GameplayCueToolset"));
+
+	if (!NewBP || !NewBP->GeneratedClass || !NewBP->GeneratedClass->GetDefaultObject())
+	{
+		UKismetSystemLibrary::RaiseScriptError(
+			FString::Printf(TEXT("Failed to create GameplayCueNotify Blueprint '%s'."), *AssetName));
+		return FString();
+	}
+
+	if (!bIsActor)
+	{
+		UGameplayCueNotify_Static* CDO = Cast<UGameplayCueNotify_Static>(NewBP->GeneratedClass->GetDefaultObject());
+		CDO->GameplayCueTag = Tag;
+		CDO->GameplayCueName = Tag.GetTagName();
+	}
+	else
+	{
+		AGameplayCueNotify_Actor* CDO = Cast<AGameplayCueNotify_Actor>(NewBP->GeneratedClass->GetDefaultObject());
+		CDO->GameplayCueTag = Tag;
+		CDO->GameplayCueName = Tag.GetTagName();
+	}
+
+	NewBP->MarkPackageDirty();
+	return NewBP->GetPathName();
+}
+
+bool UGameplayCueToolset::AddCueTag(const FString& CueTag, const FString& Comment)
+{
+	if (CueTag.IsEmpty())
+	{
+		UKismetSystemLibrary::RaiseScriptError(TEXT("CueTag cannot be empty."));
+		return false;
+	}
+
+	if (!CueTag.StartsWith(GameplayCueRootTag + TEXT(".")))
+	{
+		UKismetSystemLibrary::RaiseScriptError(FString::Printf(
+			TEXT("Gameplay cue tags must begin with '%s.' (got '%s')."),
+			*GameplayCueRootTag, *CueTag));
+		return false;
+	}
+
+	IGameplayTagsEditorModule& EditorModule =
+		FModuleManager::LoadModuleChecked<IGameplayTagsEditorModule>("GameplayTagsEditor");
+
+	const bool bSuccess = EditorModule.AddNewGameplayTagToINI(CueTag, Comment, NAME_None);
+	if (!bSuccess)
+	{
+		UKismetSystemLibrary::RaiseScriptError(
+			FString::Printf(TEXT("Failed to add gameplay cue tag '%s'."), *CueTag));
+	}
+	return bSuccess;
+}
+
+bool UGameplayCueToolset::RemoveCueTag(const FString& CueTag)
+{
+	UGameplayTagsManager& TagManager = UGameplayTagsManager::Get();
+	TSharedPtr<FGameplayTagNode> TagNode = TagManager.FindTagNode(FName(*CueTag));
+
+	if (!TagNode.IsValid())
+	{
+		UKismetSystemLibrary::RaiseScriptError(
+			FString::Printf(TEXT("Gameplay cue tag '%s' does not exist."), *CueTag));
+		return false;
+	}
+
+	IGameplayTagsEditorModule& EditorModule =
+		FModuleManager::LoadModuleChecked<IGameplayTagsEditorModule>("GameplayTagsEditor");
+
+	const bool bSuccess = EditorModule.DeleteTagFromINI(TagNode);
+	if (!bSuccess)
+	{
+		UKismetSystemLibrary::RaiseScriptError(
+			FString::Printf(TEXT("Failed to remove gameplay cue tag '%s'."), *CueTag));
+	}
+	return bSuccess;
+}
+
+TArray<FString> UGameplayCueToolset::FindCueTagsWithoutNotifies()
+{
+	TArray<FString> AllCueTags = ListCues(FString());
+	if (AllCueTags.IsEmpty())
+	{
+		return {};
+	}
+
+	TArray<FGameplayCueNotifyInfo> AllNotifies = FindCueNotifyAssets(FString());
+	TSet<FString> TagsWithNotifies;
+	for (const FGameplayCueNotifyInfo& Info : AllNotifies)
+	{
+		TagsWithNotifies.Add(Info.CueTag);
+	}
+
+	TArray<FString> Unimplemented;
+	for (const FString& Tag : AllCueTags)
+	{
+		if (!TagsWithNotifies.Contains(Tag))
+		{
+			Unimplemented.Add(Tag);
+		}
+	}
+
+	Unimplemented.Sort();
+	return Unimplemented;
+}
diff --git a/Engine/Plugins/Experimental/Toolsets/GASToolsets/Source/GASToolsets/Private/GameplayCueToolset.h b/Engine/Plugins/Experimental/Toolsets/GASToolsets/Source/GASToolsets/Private/GameplayCueToolset.h
new file mode 100644
--- /dev/null
+++ b/Engine/Plugins/Experimental/Toolsets/GASToolsets/Source/GASToolsets/Private/GameplayCueToolset.h
@@ -0,0 +1,148 @@
+// Copyright Epic Games, Inc. All Rights Reserved.
+
+#pragma once
+
+#include "CoreMinimal.h"
+#include "ToolsetRegistry/UToolsetRegistry.h"
+
+#include "GameplayCueToolset.generated.h"
+
+/// Describes a registered gameplay cue and its associated notify asset.
+USTRUCT(BlueprintType)
+struct FGameplayCueInfo
+{
+	GENERATED_BODY()
+
+	/// The fully-qualified gameplay cue tag (e.g. "GameplayCue.Character.Death").
+	UPROPERTY(BlueprintReadWrite, Category = "GameplayCue")
+	FString Tag;
+
+	/// The object path to the associated GameplayCueNotify asset.
+	/// Empty if no notify has been found for this tag in the asset registry.
+	UPROPERTY(BlueprintReadWrite, Category = "GameplayCue")
+	FString NotifyAssetPath;
+
+	/// The notify base class: "Static", "Actor", or "None" if no notify is found.
+	UPROPERTY(BlueprintReadWrite, Category = "GameplayCue")
+	FString NotifyType;
+};
+
+/// Describes a GameplayCueNotify asset found in the project.
+USTRUCT(BlueprintType)
+struct FGameplayCueNotifyInfo
+{
+	GENERATED_BODY()
+
+	/// The gameplay cue tag this notify responds to.
+	UPROPERTY(BlueprintReadWrite, Category = "GameplayCue")
+	FString CueTag;
+
+	/// The full object path to the asset (e.g. "/Game/Effects/GCN_Death.GCN_Death").
+	UPROPERTY(BlueprintReadWrite, Category = "GameplayCue")
+	FString AssetPath;
+
+	/// The asset name without the package path (e.g. "GCN_Death").
+	UPROPERTY(BlueprintReadWrite, Category = "GameplayCue")
+	FString AssetName;
+
+	/// "Static" for UGameplayCueNotify_Static subclasses, "Actor" for AGameplayCueNotify_Actor subclasses.
+	UPROPERTY(BlueprintReadWrite, Category = "GameplayCue")
+	FString NotifyType;
+};
+
+/// Provides tools for inspecting, executing, and managing gameplay cues and their notify assets.
+UCLASS(BlueprintType, Hidden)
+class UGameplayCueToolset : public UToolsetDefinition
+{
+	GENERATED_BODY()
+
+public:
+	/**
+	 * Returns gameplay cue tags registered in the project.
+	 * @param ParentTag If non-empty, only cues descending from this tag are returned.
+	 *   For example, passing "GameplayCue.Character" returns all cues under that namespace.
+	 *   Pass an empty string to return all tags under the "GameplayCue" root.
+	 * @return A sorted list of fully-qualified gameplay cue tag names.
+	 */
+	UFUNCTION(BlueprintCallable, Category = "GameplayCues")
+	static TArray<FString> ListCues(const FString& ParentTag);
+
+	/**
+	 * Returns information about a specific gameplay cue, including its notify asset.
+	 * @param CueTag The fully-qualified gameplay cue tag, e.g. "GameplayCue.Character.Death".
+	 * @return The cue's notify asset path and type ("Static", "Actor", or "None").
+	 *   Raises a script error if the tag does not exist.
+	 */
+	UFUNCTION(BlueprintCallable, Category = "GameplayCues")
+	static FGameplayCueInfo GetCueInfo(const FString& CueTag);
+
+	/**
+	 * Executes a gameplay cue non-replicated on the currently selected actor in the editor.
+	 * Useful for previewing cue effects without network replication.
+	 * Requires a PIE session or a configured GameplayCueManager to produce visible results.
+	 * @param CueTag The fully-qualified tag of the cue to execute, e.g. "GameplayCue.Character.Death".
+	 * @param NormalizedMagnitude A normalized (0.0-1.0) magnitude value passed to the cue.
+	 * @param Location World-space location parameter passed to the cue.
+	 * @param Normal World-space direction parameter passed to the cue.
+	 * @return True if the cue was dispatched. Raises a script error if no actor is selected
+	 *   or the tag does not exist.
+	 */
+	UFUNCTION(BlueprintCallable, Category = "GameplayCues")
+	static bool ExecuteCueOnSelectedActor(
+		const FString& CueTag, float NormalizedMagnitude, FVector Location, FVector Normal);
+
+	/**
+	 * Returns all GameplayCueNotify assets found in the project via the asset registry.
+	 * @param ParentTag If non-empty, only notifies whose cue tag descends from this tag are returned.
+	 *   Pass an empty string to return all notify assets in the project.
+	 * @return A list of notify descriptors, sorted by cue tag.
+	 */
+	UFUNCTION(BlueprintCallable, Category = "GameplayCues")
+	static TArray<FGameplayCueNotifyInfo> FindCueNotifyAssets(const FString& ParentTag);
+
+	/**
+	 * Creates a new GameplayCueNotify Blueprint asset at the specified content browser location.
+	 * This should ONLY be called after getting explicit direction or permission from the user.
+	 * @param CueTag The gameplay cue tag the notify should respond to.
+	 *   The tag must already exist; create it first with AddCueTag if needed.
+	 * @param PackagePath The content browser folder for the asset, e.g. "/Game/Effects/Cues".
+	 * @param AssetName The file name for the new asset, e.g. "GCN_CharacterDeath".
+	 * @param bIsActor If true, creates a GameplayCueNotify_Actor (spawned actor in the world).
+	 *   If false, creates a GameplayCueNotify_Static (instant effect, no spawned actor).
+	 * @return The object path of the created asset, or an empty string on failure.
+	 */
+	UFUNCTION(BlueprintCallable, Category = "GameplayCues")
+	static FString CreateCueNotifyAsset(
+		const FString& CueTag, const FString& PackagePath,
+		const FString& AssetName, bool bIsActor);
+
+	/**
+	 * Adds a new gameplay cue tag to the project.
+	 * This should ONLY be called after getting explicit direction or permission from the user.
+	 * @param CueTag The fully-qualified tag to add. Must begin with "GameplayCue."
+	 *   (e.g. "GameplayCue.Character.Death").
+	 * @param Comment An optional developer comment describing the cue's purpose.
+	 * @return True if the tag was added successfully.
+	 */
+	UFUNCTION(BlueprintCallable, Category = "GameplayCues")
+	static bool AddCueTag(const FString& CueTag, const FString& Comment);
+
+	/**
+	 * Removes a gameplay cue tag from the project.
+	 * This should ONLY be called after getting explicit direction or permission from the user.
+	 * @param CueTag The fully-qualified cue tag to remove (e.g. "GameplayCue.Character.Death").
+	 * @return True if the tag was removed successfully.
+	 */
+	UFUNCTION(BlueprintCallable, Category = "GameplayCues")
+	static bool RemoveCueTag(const FString& CueTag);
+
+	/**
+	 * Returns gameplay cue tags that have no corresponding GameplayCueNotify asset in the project.
+	 * Tags without notifies produce no visible effect when triggered at runtime.
+	 * @return A sorted list of gameplay cue tag names with no associated notify assets.
+	 */
+	UFUNCTION(BlueprintCallable, Category = "GameplayCues")
+	static TArray<FString> FindCueTagsWithoutNotifies();
+
+	friend class FGameplayCueToolsetSpec;
+};
diff --git a/Engine/Plugins/Experimental/Toolsets/GASToolsets/Source/GASToolsets/Private/Tests/AbilitySystemInspectorToolsetTest.cpp b/Engine/Plugins/Experimental/Toolsets/GASToolsets/Source/GASToolsets/Private/Tests/AbilitySystemInspectorToolsetTest.cpp
new file mode 100644
--- /dev/null
+++ b/Engine/Plugins/Experimental/Toolsets/GASToolsets/Source/GASToolsets/Private/Tests/AbilitySystemInspectorToolsetTest.cpp
@@ -0,0 +1,231 @@
+// Copyright Epic Games, Inc. All Rights Reserved.
+
+#include "AbilitySystemInspectorToolsetTest.h"
+
+#include "AbilitySystemInspectorToolset.h"
+#include "AbilitySystemComponent.h"
+#include "Editor.h"
+#include "Engine/World.h"
+#include "Misc/AutomationTest.h"
+
+#if WITH_DEV_AUTOMATION_TESTS
+
+BEGIN_DEFINE_SPEC(FAbilitySystemInspectorToolsetSpec, "AI.GASToolsets.AbilitySystemInspectorToolset",
+	EAutomationTestFlags::EditorContext | EAutomationTestFlags::EngineFilter)
+	AActor* TestActor = nullptr;
+END_DEFINE_SPEC(FAbilitySystemInspectorToolsetSpec)
+
+void FAbilitySystemInspectorToolsetSpec::Define()
+{
+	// -------------------------------------------------------------------------
+	// Null actor: all tools must return empty.
+	// -------------------------------------------------------------------------
+	Describe(TEXT("With a null actor"), [this]()
+	{
+		It(TEXT("GetAttributeValues returns empty"), [this]()
+		{
+			TArray<FRuntimeAttributeValue> Values =
+				UAbilitySystemInspectorToolset::GetAttributeValues(nullptr);
+			TestEqual(TEXT("Empty"), Values.Num(), 0);
+		});
+
+		It(TEXT("GetActiveEffects returns empty"), [this]()
+		{
+			TArray<FActiveEffectInfo> Effects =
+				UAbilitySystemInspectorToolset::GetActiveEffects(nullptr);
+			TestEqual(TEXT("Empty"), Effects.Num(), 0);
+		});
+
+		It(TEXT("GetGrantedAbilities returns empty"), [this]()
+		{
+			TArray<FGrantedAbilityInfo> Abilities =
+				UAbilitySystemInspectorToolset::GetGrantedAbilities(nullptr);
+			TestEqual(TEXT("Empty"), Abilities.Num(), 0);
+		});
+
+		It(TEXT("GetActiveTags returns empty"), [this]()
+		{
+			TArray<FString> Tags = UAbilitySystemInspectorToolset::GetActiveTags(nullptr);
+			TestEqual(TEXT("Empty"), Tags.Num(), 0);
+		});
+	});
+
+	// -------------------------------------------------------------------------
+	// Actor with no ASC: all tools must return empty.
+	// -------------------------------------------------------------------------
+	Describe(TEXT("With an actor that has no AbilitySystemComponent"), [this]()
+	{
+		BeforeEach([this]()
+		{
+			if (!GEditor) return;
+
+			UWorld* World = GEditor->GetEditorWorldContext().World();
+			if (!World) return;
+
+			FActorSpawnParameters Params;
+			Params.ObjectFlags = RF_Transient;
+			TestActor = World->SpawnActor<AActor>(AActor::StaticClass(),
+				FTransform::Identity, Params);
+		});
+
+		AfterEach([this]()
+		{
+			if (TestActor)
+			{
+				TestActor->Destroy();
+				TestActor = nullptr;
+			}
+		});
+
+		It(TEXT("GetAttributeValues returns empty"), [this]()
+		{
+			if (!TestNotNull(TEXT("TestActor"), TestActor)) return;
+			TArray<FRuntimeAttributeValue> Values =
+				UAbilitySystemInspectorToolset::GetAttributeValues(TestActor);
+			TestEqual(TEXT("Empty"), Values.Num(), 0);
+		});
+
+		It(TEXT("GetActiveEffects returns empty"), [this]()
+		{
+			if (!TestNotNull(TEXT("TestActor"), TestActor)) return;
+			TArray<FActiveEffectInfo> Effects =
+				UAbilitySystemInspectorToolset::GetActiveEffects(TestActor);
+			TestEqual(TEXT("Empty"), Effects.Num(), 0);
+		});
+
+		It(TEXT("GetGrantedAbilities returns empty"), [this]()
+		{
+			if (!TestNotNull(TEXT("TestActor"), TestActor)) return;
+			TArray<FGrantedAbilityInfo> Abilities =
+				UAbilitySystemInspectorToolset::GetGrantedAbilities(TestActor);
+			TestEqual(TEXT("Empty"), Abilities.Num(), 0);
+		});
+
+		It(TEXT("GetActiveTags returns empty"), [this]()
+		{
+			if (!TestNotNull(TEXT("TestActor"), TestActor)) return;
+			TArray<FString> Tags = UAbilitySystemInspectorToolset::GetActiveTags(TestActor);
+			TestEqual(TEXT("Empty"), Tags.Num(), 0);
+		});
+	});
+
+	// -------------------------------------------------------------------------
+	// Actor with a live ASC and the test attribute set registered.
+	// -------------------------------------------------------------------------
+	Describe(TEXT("With an actor that has an AbilitySystemComponent"), [this]()
+	{
+		BeforeEach([this]()
+		{
+			if (!GEditor) return;
+
+			UWorld* World = GEditor->GetEditorWorldContext().World();
+			if (!World) return;
+
+			FActorSpawnParameters Params;
+			Params.ObjectFlags = RF_Transient;
+			TestActor = World->SpawnActor<AGASToolsetsTestActor>(
+				AGASToolsetsTestActor::StaticClass(), FTransform::Identity, Params);
+
+			if (AGASToolsetsTestActor* GASActor = Cast<AGASToolsetsTestActor>(TestActor))
+			{
+				GASActor->AbilitySystemComponent->InitStats(
+					UGASToolsetsTestAttributeSet::StaticClass(), nullptr);
+			}
+		});
+
+		AfterEach([this]()
+		{
+			if (TestActor)
+			{
+				TestActor->Destroy();
+				TestActor = nullptr;
+			}
+		});
+
+		// --- GetAttributeValues --------------------------------------------------
+
+		It(TEXT("GetAttributeValues returns one entry per attribute"), [this]()
+		{
+			if (!TestNotNull(TEXT("TestActor"), TestActor)) return;
+			TArray<FRuntimeAttributeValue> Values =
+				UAbilitySystemInspectorToolset::GetAttributeValues(TestActor);
+			TestEqual(TEXT("Two attributes"), Values.Num(), 2);
+		});
+
+		It(TEXT("GetAttributeValues includes Health and MaxHealth"), [this]()
+		{
+			if (!TestNotNull(TEXT("TestActor"), TestActor)) return;
+			TArray<FRuntimeAttributeValue> Values =
+				UAbilitySystemInspectorToolset::GetAttributeValues(TestActor);
+
+			const bool bHasHealth = Values.ContainsByPredicate(
+				[](const FRuntimeAttributeValue& V)
+				{
+					return V.AttributeName == TEXT("Health");
+				});
+			const bool bHasMaxHealth = Values.ContainsByPredicate(
+				[](const FRuntimeAttributeValue& V)
+				{
+					return V.AttributeName == TEXT("MaxHealth");
+				});
+
+			TestTrue(TEXT("Has Health"), bHasHealth);
+			TestTrue(TEXT("Has MaxHealth"), bHasMaxHealth);
+		});
+
+		It(TEXT("GetAttributeValues populates SetClassName and FullName"), [this]()
+		{
+			if (!TestNotNull(TEXT("TestActor"), TestActor)) return;
+			TArray<FRuntimeAttributeValue> Values =
+				UAbilitySystemInspectorToolset::GetAttributeValues(TestActor);
+			if (!TestTrue(TEXT("Has values"), Values.Num() > 0)) return;
+
+			for (const FRuntimeAttributeValue& Value : Values)
+			{
+				TestEqual(TEXT("SetClassName"), Value.SetClassName,
+					TEXT("GASToolsetsTestAttributeSet"));
+				TestEqual(TEXT("FullName"),
+					Value.FullName,
+					FString::Printf(TEXT("GASToolsetsTestAttributeSet.%s"), *Value.AttributeName));
+			}
+		});
+
+		It(TEXT("GetAttributeValues results are sorted by FullName"), [this]()
+		{
+			if (!TestNotNull(TEXT("TestActor"), TestActor)) return;
+			TArray<FRuntimeAttributeValue> Values =
+				UAbilitySystemInspectorToolset::GetAttributeValues(TestActor);
+			for (int32 i = 1; i < Values.Num(); ++i)
+			{
+				TestTrue(TEXT("Sorted"), Values[i - 1].FullName <= Values[i].FullName);
+			}
+		});
+
+		// --- Baseline empty results for other tools ------------------------------
+
+		It(TEXT("GetActiveEffects returns empty with no active effects"), [this]()
+		{
+			if (!TestNotNull(TEXT("TestActor"), TestActor)) return;
+			TArray<FActiveEffectInfo> Effects =
+				UAbilitySystemInspectorToolset::GetActiveEffects(TestActor);
+			TestEqual(TEXT("No effects"), Effects.Num(), 0);
+		});
+
+		It(TEXT("GetGrantedAbilities returns empty with no granted abilities"), [this]()
+		{
+			if (!TestNotNull(TEXT("TestActor"), TestActor)) return;
+			TArray<FGrantedAbilityInfo> Abilities =
+				UAbilitySystemInspectorToolset::GetGrantedAbilities(TestActor);
+			TestEqual(TEXT("No abilities"), Abilities.Num(), 0);
+		});
+
+		It(TEXT("GetActiveTags returns empty with no tags"), [this]()
+		{
+			if (!TestNotNull(TEXT("TestActor"), TestActor)) return;
+			TArray<FString> Tags = UAbilitySystemInspectorToolset::GetActiveTags(TestActor);
+			TestEqual(TEXT("No tags"), Tags.Num(), 0);
+		});
+	});
+}
+
+#endif  // WITH_DEV_AUTOMATION_TESTS
diff --git a/Engine/Plugins/Experimental/Toolsets/GASToolsets/Source/GASToolsets/Private/Tests/AbilitySystemInspectorToolsetTest.h b/Engine/Plugins/Experimental/Toolsets/GASToolsets/Source/GASToolsets/Private/Tests/AbilitySystemInspectorToolsetTest.h
new file mode 100644
--- /dev/null
+++ b/Engine/Plugins/Experimental/Toolsets/GASToolsets/Source/GASToolsets/Private/Tests/AbilitySystemInspectorToolsetTest.h
@@ -0,0 +1,33 @@
+// Copyright Epic Games, Inc. All Rights Reserved.
+
+#pragma once
+
+#include "CoreMinimal.h"
+#include "AbilitySystemComponent.h"
+#include "AbilitySystemInterface.h"
+#include "Tests/AttributeSetToolsetTest.h"
+
+#include "AbilitySystemInspectorToolsetTest.generated.h"
+
+/// Minimal actor used in automation tests.
+/// Owns an AbilitySystemComponent and the test attribute set.
+UCLASS(Hidden, MinimalAPI)
+class AGASToolsetsTestActor : public AActor, public IAbilitySystemInterface
+{
+	GENERATED_BODY()
+
+public:
+	AGASToolsetsTestActor()
+	{
+		AbilitySystemComponent =
+			CreateDefaultSubobject<UAbilitySystemComponent>(TEXT("AbilitySystemComponent"));
+	}
+
+	virtual UAbilitySystemComponent* GetAbilitySystemComponent() const override
+	{
+		return AbilitySystemComponent;
+	}
+
+	UPROPERTY()
+	TObjectPtr<UAbilitySystemComponent> AbilitySystemComponent;
+};
diff --git a/Engine/Plugins/Experimental/Toolsets/GASToolsets/Source/GASToolsets/Private/Tests/AttributeSetToolsetTest.cpp b/Engine/Plugins/Experimental/Toolsets/GASToolsets/Source/GASToolsets/Private/Tests/AttributeSetToolsetTest.cpp
new file mode 100644
--- /dev/null
+++ b/Engine/Plugins/Experimental/Toolsets/GASToolsets/Source/GASToolsets/Private/Tests/AttributeSetToolsetTest.cpp
@@ -0,0 +1,138 @@
+// Copyright Epic Games, Inc. All Rights Reserved.
+
+#include "AttributeSetToolsetTest.h"
+
+#include "AttributeSetToolset.h"
+#include "Misc/AutomationTest.h"
+
+#if WITH_DEV_AUTOMATION_TESTS
+
+BEGIN_DEFINE_SPEC(FAttributeSetToolsetSpec, "AI.GASToolsets.AttributeSetToolset",
+	EAutomationTestFlags::EditorContext | EAutomationTestFlags::EngineFilter)
+END_DEFINE_SPEC(FAttributeSetToolsetSpec)
+
+void FAttributeSetToolsetSpec::Define()
+{
+	Describe(TEXT("FindAttributeSetClasses"), [this]()
+	{
+		It(TEXT("Returns at least one class"), [this]()
+		{
+			TArray<FAttributeSetClassInfo> Classes = UAttributeSetToolset::FindAttributeSetClasses();
+			TestTrue(TEXT("At least one class found"), Classes.Num() > 0);
+		});
+
+		It(TEXT("Includes the test attribute set class"), [this]()
+		{
+			TArray<FAttributeSetClassInfo> Classes = UAttributeSetToolset::FindAttributeSetClasses();
+			const FAttributeSetClassInfo* Found = Classes.FindByPredicate(
+				[](const FAttributeSetClassInfo& Info)
+				{
+					return Info.ClassName == TEXT("GASToolsetsTestAttributeSet");
+				});
+			TestNotNull(TEXT("Test class present"), Found);
+		});
+
+		It(TEXT("Reports correct attribute count for the test class"), [this]()
+		{
+			TArray<FAttributeSetClassInfo> Classes = UAttributeSetToolset::FindAttributeSetClasses();
+			const FAttributeSetClassInfo* Found = Classes.FindByPredicate(
+				[](const FAttributeSetClassInfo& Info)
+				{
+					return Info.ClassName == TEXT("GASToolsetsTestAttributeSet");
+				});
+			if (!TestNotNull(TEXT("Test class present"), Found)) return;
+
+			TestEqual(TEXT("Two attributes"), Found->Attributes.Num(), 2);
+		});
+
+		It(TEXT("Populates FullName and SetClassName on each attribute"), [this]()
+		{
+			TArray<FAttributeSetClassInfo> Classes = UAttributeSetToolset::FindAttributeSetClasses();
+			const FAttributeSetClassInfo* Found = Classes.FindByPredicate(
+				[](const FAttributeSetClassInfo& Info)
+				{
+					return Info.ClassName == TEXT("GASToolsetsTestAttributeSet");
+				});
+			if (!TestNotNull(TEXT("Test class present"), Found)) return;
+
+			for (const FGameplayAttributeInfo& Attr : Found->Attributes)
+			{
+				TestFalse(TEXT("AttributeName is not empty"), Attr.AttributeName.IsEmpty());
+				TestEqual(TEXT("SetClassName matches class"), Attr.SetClassName,
+					TEXT("GASToolsetsTestAttributeSet"));
+				TestTrue(TEXT("FullName starts with class name"),
+					Attr.FullName.StartsWith(TEXT("GASToolsetsTestAttributeSet.")));
+			}
+		});
+
+		It(TEXT("Results are sorted by class name"), [this]()
+		{
+			TArray<FAttributeSetClassInfo> Classes = UAttributeSetToolset::FindAttributeSetClasses();
+			for (int32 i = 1; i < Classes.Num(); ++i)
+			{
+				TestTrue(TEXT("Classes are sorted"),
+					Classes[i - 1].ClassName <= Classes[i].ClassName);
+			}
+		});
+
+		It(TEXT("Does not include the abstract UAttributeSet base class"), [this]()
+		{
+			TArray<FAttributeSetClassInfo> Classes = UAttributeSetToolset::FindAttributeSetClasses();
+			const bool bContainsBase = Classes.ContainsByPredicate(
+				[](const FAttributeSetClassInfo& Info)
+				{
+					return Info.ClassName == TEXT("AttributeSet");
+				});
+			TestFalse(TEXT("Base class excluded"), bContainsBase);
+		});
+	});
+
+	Describe(TEXT("ListAttributes"), [this]()
+	{
+		It(TEXT("Returns the correct number of attributes for a valid class"), [this]()
+		{
+			TArray<FGameplayAttributeInfo> Attrs =
+				UAttributeSetToolset::ListAttributes(TEXT("GASToolsetsTestAttributeSet"));
+			TestEqual(TEXT("Two attributes"), Attrs.Num(), 2);
+		});
+
+		It(TEXT("Includes Health and MaxHealth attributes"), [this]()
+		{
+			TArray<FGameplayAttributeInfo> Attrs =
+				UAttributeSetToolset::ListAttributes(TEXT("GASToolsetsTestAttributeSet"));
+
+			const bool bHasHealth = Attrs.ContainsByPredicate(
+				[](const FGameplayAttributeInfo& Info)
+				{
+					return Info.AttributeName == TEXT("Health");
+				});
+			const bool bHasMaxHealth = Attrs.ContainsByPredicate(
+				[](const FGameplayAttributeInfo& Info)
+				{
+					return Info.AttributeName == TEXT("MaxHealth");
+				});
+
+			TestTrue(TEXT("Has Health"), bHasHealth);
+			TestTrue(TEXT("Has MaxHealth"), bHasMaxHealth);
+		});
+
+		It(TEXT("FullName is formatted as SetClass.AttributeName"), [this]()
+		{
+			TArray<FGameplayAttributeInfo> Attrs =
+				UAttributeSetToolset::ListAttributes(TEXT("GASToolsetsTestAttributeSet"));
+			if (!TestTrue(TEXT("Has attributes"), Attrs.Num() > 0)) return;
+
+			for (const FGameplayAttributeInfo& Attr : Attrs)
+			{
+				const FString ExpectedFullName =
+					FString::Printf(TEXT("GASToolsetsTestAttributeSet.%s"), *Attr.AttributeName);
+				TestEqual(TEXT("FullName format"), Attr.FullName, ExpectedFullName);
+				TestEqual(TEXT("SetClassName"), Attr.SetClassName,
+					TEXT("GASToolsetsTestAttributeSet"));
+			}
+		});
+
+	});
+}
+
+#endif  // WITH_DEV_AUTOMATION_TESTS
diff --git a/Engine/Plugins/Experimental/Toolsets/GASToolsets/Source/GASToolsets/Private/Tests/AttributeSetToolsetTest.h b/Engine/Plugins/Experimental/Toolsets/GASToolsets/Source/GASToolsets/Private/Tests/AttributeSetToolsetTest.h
new file mode 100644
--- /dev/null
+++ b/Engine/Plugins/Experimental/Toolsets/GASToolsets/Source/GASToolsets/Private/Tests/AttributeSetToolsetTest.h
@@ -0,0 +1,23 @@
+// Copyright Epic Games, Inc. All Rights Reserved.
+
+#pragma once
+
+#include "CoreMinimal.h"
+#include "AttributeSet.h"
+
+#include "AttributeSetToolsetTest.generated.h"
+
+/// Minimal AttributeSet subclass used exclusively in automation tests.
+/// Defines two FGameplayAttributeData properties to verify attribute enumeration.
+UCLASS(Hidden, MinimalAPI)
+class UGASToolsetsTestAttributeSet : public UAttributeSet
+{
+	GENERATED_BODY()
+
+public:
+	UPROPERTY(BlueprintReadOnly, Category = "Attribute")
+	FGameplayAttributeData Health;
+
+	UPROPERTY(BlueprintReadOnly, Category = "Attribute")
+	FGameplayAttributeData MaxHealth;
+};
diff --git a/Engine/Plugins/Experimental/Toolsets/GASToolsets/Source/GASToolsets/Public/GASToolsets.h b/Engine/Plugins/Experimental/Toolsets/GASToolsets/Source/GASToolsets/Public/GASToolsets.h
new file mode 100644
--- /dev/null
+++ b/Engine/Plugins/Experimental/Toolsets/GASToolsets/Source/GASToolsets/Public/GASToolsets.h
@@ -0,0 +1,16 @@
+// Copyright Epic Games, Inc. All Rights Reserved.
+
+#pragma once
+
+#include "Delegates/IDelegateInstance.h"
+#include "Logging/LogMacros.h"
+#include "Modules/ModuleInterface.h"
+
+DECLARE_LOG_CATEGORY_EXTERN(LogGASToolsets, Log, Log);
+
+class FGASToolsetsModule : public IModuleInterface
+{
+public:
+	virtual void StartupModule() override;
+	virtual void ShutdownModule() override;
+};
