[Toolsets] Fix use-after-free in GASToolsets ShutdownModule during engine exit

Same issue as WorldConditionsToolset (CL 52171814). During engine exit, PurgeAllUObjectsOnExit() runs before UnloadModulesAtShutdown(), so StaticClass() returns a dangling pointer when called from ShutdownModule(). Guard with UObjectInitialized().

Changes:
- GASToolsets: guard UnregisterToolsetClass calls

#Synced-CL 52202000

#Synced-CL 52202000
#Synced-CL 52202000

diff --git a/Engine/Plugins/Experimental/Toolsets/GASToolsets/Source/GASToolsets/Private/GASToolsets.cpp b/Engine/Plugins/Experimental/Toolsets/GASToolsets/Source/GASToolsets/Private/GASToolsets.cpp
--- a/Engine/Plugins/Experimental/Toolsets/GASToolsets/Source/GASToolsets/Private/GASToolsets.cpp
+++ b/Engine/Plugins/Experimental/Toolsets/GASToolsets/Source/GASToolsets/Private/GASToolsets.cpp
@@ -19,9 +19,12 @@
 
 /*virtual*/ void FGASToolsetsModule::ShutdownModule()
 {
-	UToolsetRegistry::UnregisterToolsetClass(UAbilitySystemInspectorToolset::StaticClass());
-	UToolsetRegistry::UnregisterToolsetClass(UAttributeSetToolset::StaticClass());
-	UToolsetRegistry::UnregisterToolsetClass(UGameplayCueToolset::StaticClass());
+	if (UObjectInitialized())
+	{
+		UToolsetRegistry::UnregisterToolsetClass(UAbilitySystemInspectorToolset::StaticClass());
+		UToolsetRegistry::UnregisterToolsetClass(UAttributeSetToolset::StaticClass());
+		UToolsetRegistry::UnregisterToolsetClass(UGameplayCueToolset::StaticClass());
+	}
 }
 
 IMPLEMENT_MODULE(FGASToolsetsModule, GASToolsets);
