[AI Assistant] Disable GAS toolset plugin by default.

#Synced-CL 52269267

diff --git a/Engine/Plugins/Experimental/Toolsets/GASToolsets/GASToolsets.uplugin b/Engine/Plugins/Experimental/Toolsets/GASToolsets/GASToolsets.uplugin
--- a/Engine/Plugins/Experimental/Toolsets/GASToolsets/GASToolsets.uplugin
+++ b/Engine/Plugins/Experimental/Toolsets/GASToolsets/GASToolsets.uplugin
@@ -16,7 +16,7 @@
 	"IsExperimentalVersion": true,
 	"NoRedist": true,
 	"Installed": false,
-	"EnabledByDefault": true,
+	"EnabledByDefault": false,
 	"Modules": [
 		{
 			"Name": "GASToolsets",
