[AI Assistant Toolsets] Move toolset tests under AI.Toolsets.

#Synced-CL 52328365

diff --git a/Engine/Plugins/Experimental/Toolsets/GASToolsets/Source/GASToolsets/Private/Tests/AbilitySystemInspectorToolsetTest.cpp b/Engine/Plugins/Experimental/Toolsets/GASToolsets/Source/GASToolsets/Private/Tests/AbilitySystemInspectorToolsetTest.cpp
--- a/Engine/Plugins/Experimental/Toolsets/GASToolsets/Source/GASToolsets/Private/Tests/AbilitySystemInspectorToolsetTest.cpp
+++ b/Engine/Plugins/Experimental/Toolsets/GASToolsets/Source/GASToolsets/Private/Tests/AbilitySystemInspectorToolsetTest.cpp
@@ -10,7 +10,8 @@
 
 #if WITH_DEV_AUTOMATION_TESTS
 
-BEGIN_DEFINE_SPEC(FAbilitySystemInspectorToolsetSpec, "AI.GASToolsets.AbilitySystemInspectorToolset",
+BEGIN_DEFINE_SPEC(FAbilitySystemInspectorToolsetSpec,
+	"AI.Toolsets.GASToolsets.AbilitySystemInspectorToolset",
 	EAutomationTestFlags::EditorContext | EAutomationTestFlags::EngineFilter)
 	AActor* TestActor = nullptr;
 END_DEFINE_SPEC(FAbilitySystemInspectorToolsetSpec)
diff --git a/Engine/Plugins/Experimental/Toolsets/GASToolsets/Source/GASToolsets/Private/Tests/AttributeSetToolsetTest.cpp b/Engine/Plugins/Experimental/Toolsets/GASToolsets/Source/GASToolsets/Private/Tests/AttributeSetToolsetTest.cpp
--- a/Engine/Plugins/Experimental/Toolsets/GASToolsets/Source/GASToolsets/Private/Tests/AttributeSetToolsetTest.cpp
+++ b/Engine/Plugins/Experimental/Toolsets/GASToolsets/Source/GASToolsets/Private/Tests/AttributeSetToolsetTest.cpp
@@ -7,7 +7,7 @@
 
 #if WITH_DEV_AUTOMATION_TESTS
 
-BEGIN_DEFINE_SPEC(FAttributeSetToolsetSpec, "AI.GASToolsets.AttributeSetToolset",
+BEGIN_DEFINE_SPEC(FAttributeSetToolsetSpec, "AI.Toolsets.GASToolsets.AttributeSetToolset",
 	EAutomationTestFlags::EditorContext | EAutomationTestFlags::EngineFilter)
 END_DEFINE_SPEC(FAttributeSetToolsetSpec)
 
diff --git a/Engine/Plugins/Experimental/Toolsets/PhysicsToolsets/Source/PhysicsToolsets/Private/PhysicsToolsets/Tests/PhysicsAssetToolsetTest.cpp b/Engine/Plugins/Experimental/Toolsets/PhysicsToolsets/Source/PhysicsToolsets/Private/PhysicsToolsets/Tests/PhysicsAssetToolsetTest.cpp
--- a/Engine/Plugins/Experimental/Toolsets/PhysicsToolsets/Source/PhysicsToolsets/Private/PhysicsToolsets/Tests/PhysicsAssetToolsetTest.cpp
+++ b/Engine/Plugins/Experimental/Toolsets/PhysicsToolsets/Source/PhysicsToolsets/Private/PhysicsToolsets/Tests/PhysicsAssetToolsetTest.cpp
@@ -13,7 +13,7 @@
 
 #if WITH_DEV_AUTOMATION_TESTS
 
-BEGIN_DEFINE_SPEC(FPhysicsAssetToolsetSpec, "AI.PhysicsToolsets.PhysicsAssetToolsetSpec",
+BEGIN_DEFINE_SPEC(FPhysicsAssetToolsetSpec, "AI.Toolsets.PhysicsToolsets.PhysicsAssetToolsetSpec",
 	PhysicsToolsetsTest::Flags)
 
 	// In-memory physics asset shared across tests.
